Settings: Removed custom reset/export/import logic for UEditorUserSettings (TTP# 326544)

diff --git a/Engine/Source/Editor/EditorSettingsViewer/Private/EditorSettingsViewerModule.cpp b/Engine/Source/Editor/EditorSettingsViewer/Private/EditorSettingsViewerModule.cpp
--- a/Engine/Source/Editor/EditorSettingsViewer/Private/EditorSettingsViewerModule.cpp
+++ b/Engine/Source/Editor/EditorSettingsViewer/Private/EditorSettingsViewerModule.cpp
@@ -83,14 +83,6 @@ protected:
 	 */
 	void RegisterGeneralSettings( ISettingsModule& SettingsModule )
 	{
-		// general settings
-		FSettingsSectionDelegates UserSettingsDelegates;
-		{
-			UserSettingsDelegates.ExportDelegate = FOnSettingsSectionExport::CreateRaw(this, &FEditorSettingsViewerModule::HandleUserSettingsExport);
-			UserSettingsDelegates.ImportDelegate = FOnSettingsSectionImport::CreateRaw(this, &FEditorSettingsViewerModule::HandleUserSettingsImport);
-			UserSettingsDelegates.ResetDefaultsDelegate = FOnSettingsSectionResetDefaults::CreateRaw(this, &FEditorSettingsViewerModule::HandleUserSettingsResetDefaults);
-		}
-
 		// automation settings
 		SettingsModule.RegisterSettings("Editor", "General", "AutomationTest",
 			LOCTEXT("AutomationSettingsName", "Automation"),
@@ -130,8 +122,7 @@ protected:
 		SettingsModule.RegisterSettings("Editor", "General", "LoadingSaving",
 			LOCTEXT("LoadingSavingSettingsName", "Loading & Saving"),
 			LOCTEXT("LoadingSavingSettingsDescription", "Change how the Editor loads and saves files."),
-			GetMutableDefault<UEditorLoadingSavingSettings>(),
-			UserSettingsDelegates
+			GetMutableDefault<UEditorLoadingSavingSettings>()
 		);
 
 		// @todo thomass: proper settings support for source control module
@@ -141,8 +132,7 @@ protected:
 		SettingsModule.RegisterSettings("Editor", "General", "UserSettings",
 			LOCTEXT("UserSettingsName", "Miscellaneous"),
 			LOCTEXT("UserSettingsDescription", "Customize the behavior, look and feel of the editor."),
-			&GEditor->AccessEditorUserSettings(),
-			UserSettingsDelegates
+			&GEditor->AccessEditorUserSettings()
 		);
 
 		// experimental features
@@ -301,42 +291,6 @@ private:
 		return false;
 	}
 
-	// Handles exporting the preferences settings to a file.
-	bool HandleUserSettingsExport( const FString& Filename )
-	{
-		FGlobalTabmanager::Get()->SaveAllVisualState();
-		GEditor->SaveEditorUserSettings();
-		return BackupFile(GEditorUserSettingsIni, *Filename);
-	}
-
-	// Handles importing the preferences settings from a file.
-	bool HandleUserSettingsImport( const FString& Filename )
-	{
-		if( EAppReturnType::Ok == ShowRestartWarning(LOCTEXT("ImportPreferences_Title", "Import Preferences")))
-		{
-			FUnrealEdMisc::Get().SetConfigRestoreFilename(Filename, GEditorUserSettingsIni);
-			FUnrealEdMisc::Get().RestartEditor(false);
-
-			return true;
-		}
-
-		return false;
-	}
-
-	// Handles resetting the preferences settings to their defaults.
-	bool HandleUserSettingsResetDefaults( )
-	{
-		if( EAppReturnType::Ok == ShowRestartWarning(LOCTEXT("ResetPreferences_Title", "Reset Preferences")))
-		{
-			FUnrealEdMisc::Get().ForceDeletePreferences(true);
-			FUnrealEdMisc::Get().RestartEditor(false);
-
-			return true;
-		};
-
-		return false;
-	}
-
 	// Handles exporting input bindings to a file
 	bool HandleInputBindingsExport( const FString& Filename )
 	{
