Miscellaneous improvements to keyboard shortcuts editor.
#ttp 335130 - EDITOR: Keyboard Shortcuts: Toast falsely warns of failed import if prompted for save
#branch UE4
#proj Editor.EditorSettingsViewer, Editor.InputBindingsEditor, Editor.UnrealEd, Developer.SettingsEditor
#add Added FEditorSettingsViewerModule::HandleInputBindingsSave as a SaveDelegate for the keyboard bindings section, whose purpose is just to remove user gestures so nothing is flushed to config at editor shutdown, and hence does not supersede the freshly copied imported config.
#change Changed the message in FEditorSettingsViewerModule::ShowRestartWarning to inform the user that changes will not take place until after restart.
#change Fixed input binding editor behavior so that valid gestures are committed immediately instead of requiring a change of focus.

diff --git a/Engine/Source/Editor/EditorSettingsViewer/Private/EditorSettingsViewerModule.cpp b/Engine/Source/Editor/EditorSettingsViewer/Private/EditorSettingsViewerModule.cpp
--- a/Engine/Source/Editor/EditorSettingsViewer/Private/EditorSettingsViewerModule.cpp
+++ b/Engine/Source/Editor/EditorSettingsViewer/Private/EditorSettingsViewerModule.cpp
@@ -95,6 +95,7 @@ protected:
 		InputBindingDelegates.ExportDelegate = FOnSettingsSectionExport::CreateRaw(this, &FEditorSettingsViewerModule::HandleInputBindingsExport);
 		InputBindingDelegates.ImportDelegate = FOnSettingsSectionImport::CreateRaw(this, &FEditorSettingsViewerModule::HandleInputBindingsImport);
 		InputBindingDelegates.ResetDefaultsDelegate = FOnSettingsSectionResetDefaults::CreateRaw(this, &FEditorSettingsViewerModule::HandleInputBindingsResetToDefault);
+		InputBindingDelegates.SaveDelegate = FOnSettingsSectionSave::CreateRaw(this, &FEditorSettingsViewerModule::HandleInputBindingsSave);
 
 		InputBindingEditorPanel = FModuleManager::LoadModuleChecked<IInputBindingEditorModule>("InputBindingEditor").CreateInputBindingEditorPanel();
 
@@ -261,7 +262,7 @@ private:
 	// Show a warning that the editor will require a restart and return its result
 	EAppReturnType::Type ShowRestartWarning(const FText& Title) const
 	{
-		return OpenMsgDlgInt(EAppMsgType::OkCancel, LOCTEXT("ActionRestartMsg", "This action requires the editor to restart; you will be prompted to save any changes. Continue?" ), Title);
+		return OpenMsgDlgInt(EAppMsgType::OkCancel, LOCTEXT("ActionRestartMsg", "Imported settings won't be applied until the editor is restarted. Do you wish to restart now (you will be prompted to save any changes)?" ), Title);
 	}
 
 	// Backup a file
@@ -336,6 +337,17 @@ private:
 		return false;
 	}
 
+	// Handles saving default input bindings.
+	// This only gets called by SSettingsEditor::HandleImportButtonClicked when importing new settings,
+	// and its implementation here is just to flush custom input bindings so that editor shutdown doesn't
+	// overwrite the imported settings just copied across.
+	bool HandleInputBindingsSave()
+	{
+		FInputBindingManager::Get().RemoveUserDefinedGestures();
+		GConfig->Flush(false, GEditorKeyBindingsIni);
+		return true;
+	}
+
 	bool HandleRegionAndLanguageExport(const FString& FileName)
 	{
 		FString CultureName = GetMutableDefault<UInternationalizationSettingsModel>()->GetCultureName();
diff --git a/Engine/Source/Editor/InputBindingEditor/Private/Widgets/SGestureEditBox.cpp b/Engine/Source/Editor/InputBindingEditor/Private/Widgets/SGestureEditBox.cpp
--- a/Engine/Source/Editor/InputBindingEditor/Private/Widgets/SGestureEditBox.cpp
+++ b/Engine/Source/Editor/InputBindingEditor/Private/Widgets/SGestureEditBox.cpp
@@ -105,7 +105,7 @@ FText SGestureEditBox::GetNotificationMessage() const
 
 void SGestureEditBox::OnGestureEditorLostFocus()
 {
-	if( FSlateApplication::Get().GetMouseCaptor() != GestureAcceptButton && !GestureEditor->IsTyping() )
+	if( (!GestureAcceptButton.IsValid() || FSlateApplication::Get().GetMouseCaptor() != GestureAcceptButton) && !GestureEditor->IsTyping() )
 	{
 		if( GestureEditor->IsEditing() && GestureEditor->IsEditedGestureValid() && !GestureEditor->HasConflict() )
 		{
@@ -143,6 +143,13 @@ void SGestureEditBox::OnGestureChanged()
 	else
 	{
 		ConflictPopup->SetIsOpen(false);
+
+		if (GestureEditor->IsEditedGestureValid())
+		{
+			GestureEditor->CommitNewGesture();
+			GestureEditor->StopEditing();
+			FSlateApplication::Get().ClearKeyboardFocus(EKeyboardFocusCause::SetDirectly);
+		}
 	}
 }
 
