Restore the Global settings tab in the editor; setting for local DDC location was added between Dev-Core and Dev-Editor merges.

diff --git a/Engine/Source/Editor/EditorSettingsViewer/Private/EditorSettingsViewerModule.cpp b/Engine/Source/Editor/EditorSettingsViewer/Private/EditorSettingsViewerModule.cpp
--- a/Engine/Source/Editor/EditorSettingsViewer/Private/EditorSettingsViewerModule.cpp
+++ b/Engine/Source/Editor/EditorSettingsViewer/Private/EditorSettingsViewerModule.cpp
@@ -145,6 +145,13 @@ protected:
 		// @todo thomass: proper settings support for source control module
 		GetMutableDefault<UEditorLoadingSavingSettings>()->SccHackInitialize();
 
+		// global editor settings
+		SettingsModule.RegisterSettings("Editor", "General", "Global",
+			LOCTEXT("GlobalSettingsName", "Global"),
+			LOCTEXT("GlobalSettingsDescription", "Edit global settings that affect all editors."),
+			GetMutableDefault<UEditorSettings>()
+		);
+
 		// misc unsorted settings
 		SettingsModule.RegisterSettings("Editor", "General", "UserSettings",
 			LOCTEXT("UserSettingsName", "Miscellaneous"),
diff --git a/Engine/Source/Editor/UnrealEd/Public/Settings/EditorSettings.h b/Engine/Source/Editor/UnrealEd/Public/Settings/EditorSettings.h
--- a/Engine/Source/Editor/UnrealEd/Public/Settings/EditorSettings.h
+++ b/Engine/Source/Editor/UnrealEd/Public/Settings/EditorSettings.h
@@ -23,10 +23,6 @@ class UNREALED_API UEditorSettings : public UObject
 	UPROPERTY()
 	bool bEditorAnalyticsEnabled_DEPRECATED;
 
-	/** Whether the Content Browser should open the Sources Panel by default */
-	UPROPERTY(EditAnywhere, config, Category=ContentBrowser)
-	bool bOpenSourcesPanelByDefault;
-
 	/** Sets the path to be used for caching derived data (native textures, compiled shaders, etc...). The editor must be restarted for changes to take effect. */
 	UPROPERTY(EditAnywhere, config, Category=DerivedData)
 	FDirectoryPath LocalDerivedDataCache;
