Fixed a typo in the settings category display name for UCrashReportsPrivacySettings (which never surfaced due to the category already being created)
#rnx

diff --git a/Engine/Source/Editor/EditorSettingsViewer/Private/EditorSettingsViewerModule.cpp b/Engine/Source/Editor/EditorSettingsViewer/Private/EditorSettingsViewerModule.cpp
--- a/Engine/Source/Editor/EditorSettingsViewer/Private/EditorSettingsViewerModule.cpp
+++ b/Engine/Source/Editor/EditorSettingsViewer/Private/EditorSettingsViewerModule.cpp
@@ -264,7 +264,7 @@ protected:
 			);
 
 		// Crash reports
-		SettingsModule.RegisterSettings("Editor", "Privacy", "Analitics",
+		SettingsModule.RegisterSettings("Editor", "Privacy", "Analytics",
 			LOCTEXT("PrivacyBugReportsSettingsName", "Bug Reports"),
 			LOCTEXT("PrivacyBugReportsSettingsDescription", "Configure the way your bug reports information is handled."),
 			GetMutableDefault<UCrashReportsPrivacySettings>()
