Fixed Usage Data Options missing from Editor Preferences
#rnx

diff --git a/Engine/Source/Editor/EditorSettingsViewer/Private/EditorSettingsViewerModule.cpp b/Engine/Source/Editor/EditorSettingsViewer/Private/EditorSettingsViewerModule.cpp
--- a/Engine/Source/Editor/EditorSettingsViewer/Private/EditorSettingsViewerModule.cpp
+++ b/Engine/Source/Editor/EditorSettingsViewer/Private/EditorSettingsViewerModule.cpp
@@ -263,8 +263,8 @@ protected:
 			GetMutableDefault<UAnalyticsPrivacySettings>()
 			);
 
-		// Crash reports
-		SettingsModule.RegisterSettings("Editor", "Privacy", "Analytics",
+		// Bug reports
+		SettingsModule.RegisterSettings("Editor", "Privacy", "BugReports",
 			LOCTEXT("PrivacyBugReportsSettingsName", "Bug Reports"),
 			LOCTEXT("PrivacyBugReportsSettingsDescription", "Configure the way your bug reports information is handled."),
 			GetMutableDefault<UCrashReportsPrivacySettings>()
