Editor: Removed a duplicate registration of the Blueprint Editor Settings (now that it is derived from UDeveloperSettings it doesn't need to be manually registered anymore)
#rnx

diff --git a/Engine/Source/Editor/EditorSettingsViewer/Private/EditorSettingsViewerModule.cpp b/Engine/Source/Editor/EditorSettingsViewer/Private/EditorSettingsViewerModule.cpp
--- a/Engine/Source/Editor/EditorSettingsViewer/Private/EditorSettingsViewerModule.cpp
+++ b/Engine/Source/Editor/EditorSettingsViewer/Private/EditorSettingsViewerModule.cpp
@@ -246,13 +246,6 @@ protected:
 			GetMutableDefault<UGraphEditorSettings>()
 		);
 
-		// graph editors
-		SettingsModule.RegisterSettings("Editor", "ContentEditors", "BlueprintEditor",
-			LOCTEXT("ContentEditorsBlueprintEditorSettingsName", "Blueprint Editor"),
-			LOCTEXT("ContentEditorsGraphBlueprintSettingsDescription", "Customize Blueprint Editors."),
-			GetMutableDefault<UBlueprintEditorSettings>()
-		);
-
 		// Persona editors
 		SettingsModule.RegisterSettings("Editor", "ContentEditors", "Persona",
 			LOCTEXT("ContentEditorsPersonaSettingsName", "Animation Editor"),
